MessagingRpc: Started to implement messaging based RPC layer

diff --git a/Engine/Source/Runtime/MessagingRpc/MessagingRpc.Build.cs b/Engine/Source/Runtime/MessagingRpc/MessagingRpc.Build.cs
new file mode 100644
--- /dev/null
+++ b/Engine/Source/Runtime/MessagingRpc/MessagingRpc.Build.cs
@@ -0,0 +1,32 @@
+// Copyright 1998-2015 Epic Games, Inc. All Rights Reserved.
+
+namespace UnrealBuildTool.Rules
+{
+	public class MessagingRpc : ModuleRules
+	{
+		public MessagingRpc(TargetInfo Target)
+		{
+			PublicDependencyModuleNames.AddRange(
+				new string[]
+				{
+					"Core",
+					"CoreUObject",
+                    "PortalServices",
+				}
+			);
+
+            PrivateDependencyModuleNames.AddRange(
+                new string[]
+                {
+                    "Messaging",
+                }
+            );
+
+			PrivateIncludePaths.AddRange(
+				new string[] {
+					"Runtime/MessagingRpc/Private",
+				}
+			);
+		}
+	}
+}
diff --git a/Engine/Source/Runtime/MessagingRpc/Private/MessageRpcClient.cpp b/Engine/Source/Runtime/MessagingRpc/Private/MessageRpcClient.cpp
new file mode 100644
--- /dev/null
+++ b/Engine/Source/Runtime/MessagingRpc/Private/MessageRpcClient.cpp
@@ -0,0 +1,106 @@
+// Copyright 1998-2015 Epic Games, Inc. All Rights Reserved.
+
+#include "PrivatePCH.h"
+
+
+/* FMessageRpcClient structors
+ *****************************************************************************/
+
+FMessageRpcClient::FMessageRpcClient()
+{
+	MessageEndpoint = FMessageEndpoint::Builder("FPortalAppWindowEndpoint")
+		.Handling<FMessageRpcProgress>(this, &FMessageRpcClient::HandleMessage);
+
+	TickerHandle = FTicker::GetCoreTicker().AddTicker(FTickerDelegate::CreateRaw(this, &FMessageRpcClient::HandleTicker), MESSAGE_RPC__RETRY_INTERVAL);
+}
+
+
+FMessageRpcClient::~FMessageRpcClient()
+{
+	FTicker::GetCoreTicker().RemoveTicker(TickerHandle);
+}
+
+
+/* FMessageRpcClient implementation
+ *****************************************************************************/
+
+TSharedPtr<IMessageRpcCall> FMessageRpcClient::FindCall(const TSharedRef<IMessageContext, ESPMode::ThreadSafe>& Context)
+{
+	auto Request = static_cast<const FRpcMessage*>(Context->GetMessage());
+	return Calls.FindRef(Request->CallId);
+}
+
+
+void FMessageRpcClient::SendCall(const TSharedPtr<IMessageRpcCall>& Call)
+{
+	MessageEndpoint->Send(
+		Call->GetMessage(),
+		Call->GetMessageType(),
+		nullptr,
+		TArrayBuilder<FMessageAddress>().Add(ServerAddress),
+		FTimespan::Zero(),
+		FDateTime::MaxValue()
+	);
+}
+
+
+/* IMessageRpcClient interface
+ *****************************************************************************/
+
+FGuid FMessageRpcClient::AddCall(const TSharedRef<IMessageRpcCall>& Call)
+{
+	FGuid CallId = FGuid::NewGuid();
+
+	Calls.Add(CallId, Call);
+	SendCall(Call);
+
+	return CallId;
+}
+
+
+void FMessageRpcClient::CancelCall(const FGuid& CallId)
+{
+	TSharedPtr<IMessageRpcCall> Call;
+	
+	if (Calls.RemoveAndCopyValue(CallId, Call))
+	{
+		MessageEndpoint->Send(new FMessageRpcCancel(CallId), ServerAddress);
+	}
+}
+
+
+/* FMessageRpcClient event handlers
+ *****************************************************************************/
+
+void FMessageRpcClient::HandleMessage(const FMessageRpcProgress& Message, const TSharedRef<IMessageContext, ESPMode::ThreadSafe>& Context)
+{
+	TSharedPtr<IMessageRpcCall> Call = FindCall(Context);
+
+	if (Call.IsValid())
+	{
+		Call->UpdateProgress(Message.Completion, FText::FromString(Message.StatusText));
+	}
+}
+
+
+bool FMessageRpcClient::HandleTicker(float DeltaTime)
+{
+	const FDateTime UtcNow = FDateTime::UtcNow();
+
+	for (TMap<FGuid, TSharedPtr<IMessageRpcCall>>::TIterator It(Calls); It; ++It)
+	{
+		auto Call = It.Value();
+		
+		if (UtcNow - Call->GetTimeCreated() > FTimespan::FromSeconds(MESSAGE_RPC_RETRY_TIMEOUT))
+		{
+			It.RemoveCurrent();
+			Call->TimeOut();
+		}
+		else if (UtcNow - Call->GetLastUpdated() > FTimespan::FromSeconds(MESSAGE_RPC__RETRY_INTERVAL))
+		{
+			SendCall(Call);
+		}
+	}
+
+	return true;
+}
diff --git a/Engine/Source/Runtime/MessagingRpc/Private/MessageRpcClient.h b/Engine/Source/Runtime/MessagingRpc/Private/MessageRpcClient.h
new file mode 100644
--- /dev/null
+++ b/Engine/Source/Runtime/MessagingRpc/Private/MessageRpcClient.h
@@ -0,0 +1,74 @@
+// Copyright 1998-2015 Epic Games, Inc. All Rights Reserved.
+
+#pragma once
+
+#include "IMessageRpcClient.h"
+
+
+class FMessageEndpoint;
+struct FMessageRpcCancel;
+struct FMessageRpcProgress;
+class IMessageContext;
+class IMessageRpcCall;
+
+
+/**
+ * Implements an RPC client.
+ */
+class FMessageRpcClient
+	: public IMessageRpcClient
+{
+public:
+
+	/** Default constructor. */
+	FMessageRpcClient();
+
+	/** Virtual destructor. */
+	virtual ~FMessageRpcClient();
+
+protected:
+
+	/**
+	 * Find the active RPC call for a received message.
+	 *
+	 * @param Context The context of the received message.
+	 * @return The matching call object, or nullptr if not found.
+	 */
+	TSharedPtr<IMessageRpcCall> FindCall(const TSharedRef<IMessageContext, ESPMode::ThreadSafe>& Context);
+
+	/**
+	 * Sends an RPC call to the server.
+	 *
+	 * @param Call The RPC call to make.
+	 */
+	void SendCall(const TSharedPtr<IMessageRpcCall>& Call);
+
+protected:
+
+	// IMessageRpcClient interface
+
+	virtual FGuid AddCall(const TSharedRef<IMessageRpcCall>& Call) override;
+	virtual void CancelCall(const FGuid& CallId) override;
+
+private:
+
+	/** Handles FMessageRpcProgress messages. */
+	void HandleMessage(const FMessageRpcProgress& Message, const TSharedRef<IMessageContext, ESPMode::ThreadSafe>& Context);
+
+	/** Handles the ticker. */
+	bool HandleTicker(float DeltaTime);
+
+private:
+
+	/** Active RPC calls. */
+	TMap<FGuid, TSharedPtr<IMessageRpcCall>> Calls;
+
+	/** Message endpoint. */
+	TSharedPtr<FMessageEndpoint, ESPMode::ThreadSafe> MessageEndpoint;
+
+	/** The RPC server's address. */
+	FMessageAddress ServerAddress;
+
+	/** Handle to the registered ticker. */
+	FDelegateHandle TickerHandle;
+};
diff --git a/Engine/Source/Runtime/MessagingRpc/Private/MessageRpcMessages.h b/Engine/Source/Runtime/MessagingRpc/Private/MessageRpcMessages.h
new file mode 100644
--- /dev/null
+++ b/Engine/Source/Runtime/MessagingRpc/Private/MessageRpcMessages.h
@@ -0,0 +1,75 @@
+// Copyright 1998-2015 Epic Games, Inc. All Rights Reserved.
+
+#pragma once
+
+#include "MessageRpcMessages.generated.h"
+
+
+/** Message for canceling an RPC call. */
+USTRUCT()
+struct FMessageRpcCancel
+{
+	GENERATED_USTRUCT_BODY()
+
+	/** Correlation identifier for the RPC call that this message refers to. */
+	UPROPERTY()
+	FGuid CallId;
+
+	/** Default constructor. */
+	FMessageRpcCancel() { }
+
+	/** Creates and initializes a new instance. */
+	FMessageRpcCancel(const FGuid& InCallId)
+		: CallId(InCallId)
+	{ }
+};
+
+
+/** Message for updating the progress of an RPC call. */
+USTRUCT()
+struct FMessageRpcProgress
+{
+	GENERATED_USTRUCT_BODY()
+
+	/** Completion percentage (0.0 to 1.0). */
+	UPROPERTY()
+	float Completion;
+
+	/** Correlation identifier for the RPC call that this message refers to. */
+	UPROPERTY()
+	FGuid CallId;
+
+	/** Status text. */
+	UPROPERTY()
+	FString StatusText;
+
+	/** Default constructor. */
+	FMessageRpcProgress() { }
+
+	/** Creates and initializes a new instance. */
+	FMessageRpcProgress(const FGuid& InCallId, float InCompletion, const FText& InStatusText)
+		: Completion(InCompletion)
+		, CallId(InCallId)
+		, StatusText(InStatusText.ToString())
+	{ }
+};
+
+
+/** Message for notifying RPC clients that a call was invalid. */
+USTRUCT()
+struct FMessageRpcInvalid
+{
+	GENERATED_USTRUCT_BODY()
+
+	/** Correlation identifier for the RPC call that this message refers to. */
+	UPROPERTY()
+	FGuid CallId;
+
+	/** Default constructor. */
+	FMessageRpcInvalid() { }
+
+	/** Creates and initializes a new instance. */
+	FMessageRpcInvalid(const FGuid& InCallId)
+		: CallId(InCallId)
+	{ }
+};
diff --git a/Engine/Source/Runtime/MessagingRpc/Private/MessageRpcServer.cpp b/Engine/Source/Runtime/MessagingRpc/Private/MessageRpcServer.cpp
new file mode 100644
--- /dev/null
+++ b/Engine/Source/Runtime/MessagingRpc/Private/MessageRpcServer.cpp
@@ -0,0 +1,147 @@
+// Copyright 1998-2015 Epic Games, Inc. All Rights Reserved.
+
+#include "PrivatePCH.h"
+#include "IMessageRpcHandler.h"
+#include "IMessageRpcReturn.h"
+
+
+/* FMessageRpcServer structors
+ *****************************************************************************/
+
+FMessageRpcServer::FMessageRpcServer()
+{
+	MessageEndpoint = FMessageEndpoint::Builder("FPortalAppWindowEndpoint")
+		.WithCatchall(this, &FMessageRpcServer::HandleMessage);
+
+	TickerHandle = FTicker::GetCoreTicker().AddTicker(FTickerDelegate::CreateRaw(this, &FMessageRpcServer::HandleTicker), 0.1f);
+}
+
+
+FMessageRpcServer::~FMessageRpcServer()
+{
+	FTicker::GetCoreTicker().RemoveTicker(TickerHandle);
+}
+
+
+/* IMessageRpcServer interface
+ *****************************************************************************/
+
+void FMessageRpcServer::AddHandler(const FName& RequestMessageType, const TSharedRef<IMessageRpcHandler>& Handler)
+{
+	Handlers.Add(RequestMessageType, Handler);
+}
+
+
+/* FMessageRpcServer implementation
+ *****************************************************************************/
+
+void FMessageRpcServer::ProcessCancelation(const FMessageRpcCancel& Message, const TSharedRef<IMessageContext, ESPMode::ThreadSafe>& Context)
+{
+	FReturnInfo ReturnInfo;
+
+	if (Returns.RemoveAndCopyValue(Message.CallId, ReturnInfo))
+	{
+		ReturnInfo.Return->Cancel();
+	}
+}
+
+
+void FMessageRpcServer::ProcessRequest(const TSharedRef<IMessageContext, ESPMode::ThreadSafe>& Context)
+{
+	auto Message = (FRpcMessage*)Context->GetMessage();
+	auto Handler = Handlers.FindRef(Context->GetMessageType());
+
+	if (Handler.IsValid())
+	{
+		/*
+		TSharedPtr<IMessageRpcServer::IPortalRpcReturn> Return = Handler->HandleRequest(Context);
+
+		if (Return.IsValid())
+		{
+			Returns.Add(Message->CallId, Return);
+
+			return;
+		}*/
+	}
+
+	// notify caller that call was not handled
+	MessageEndpoint->Send(new FMessageRpcInvalid(Message->CallId), Context->GetSender());
+}
+
+
+void FMessageRpcServer::SendProgress(const FGuid& CallId, const FReturnInfo& ReturnInfo)
+{
+	const TSharedPtr<IAsyncProgress>& Progress = ReturnInfo.Progress;
+	const TSharedPtr<IAsyncTask>& Task = ReturnInfo.Task;
+
+	MessageEndpoint->Send(
+		new FMessageRpcProgress(
+			CallId,
+			Progress.IsValid() ? Progress->GetCompletion().Get(-1.0f) : -1.0f,
+			Progress.IsValid() ? Progress->GetStatusText() : FText::GetEmpty()
+		),
+		ReturnInfo.ClientAddress
+	);
+}
+
+
+void FMessageRpcServer::SendResult(const FGuid& CallId, const FReturnInfo& ReturnInfo)
+{
+	FRpcMessage* Message = ReturnInfo.Return->CreateResponseMessage();
+
+	MessageEndpoint->Send(
+		Message,
+		ReturnInfo.Return->GetResponseTypeInfo(),
+		nullptr,
+		TArrayBuilder<FMessageAddress>().Add(ReturnInfo.ClientAddress),
+		FTimespan::Zero(),
+		FDateTime::MaxValue()
+	);
+}
+
+
+/* FMessageRpcServer event handlers
+ *****************************************************************************/
+
+void FMessageRpcServer::HandleMessage(const TSharedRef<IMessageContext, ESPMode::ThreadSafe>& Context)
+{
+	const TWeakObjectPtr<UScriptStruct>& MessageTypeInfo = Context->GetMessageTypeInfo();
+
+	if (!MessageTypeInfo.IsValid())
+	{
+		return;
+	}
+
+	if (MessageTypeInfo == FMessageRpcCancel::StaticStruct())
+	{
+		ProcessCancelation(*static_cast<const FMessageRpcCancel*>(Context->GetMessage()), Context);
+	}
+	else if (MessageTypeInfo->IsChildOf(FRpcMessage::StaticStruct()))
+	{
+		ProcessRequest(Context);
+	}
+}
+
+
+bool FMessageRpcServer::HandleTicker(float DeltaTime)
+{
+	const FDateTime UtcNow = FDateTime::UtcNow();
+
+	for (TMap<FGuid, FReturnInfo>::TIterator It(Returns); It; ++It)
+	{
+		FReturnInfo& ReturnInfo = It.Value();
+
+		if (ReturnInfo.Return->IsReady())
+		{
+			SendResult(It.Key(), ReturnInfo);
+			It.RemoveCurrent();
+		}
+		else if (UtcNow - ReturnInfo.LastProgressSent > FTimespan::FromSeconds(MESSAGE_RPC__RETRY_INTERVAL * 0.25))
+		{
+			SendProgress(It.Key(), ReturnInfo);
+			ReturnInfo.LastProgressSent = UtcNow;
+		}
+	}
+
+	return true;
+}
diff --git a/Engine/Source/Runtime/MessagingRpc/Private/MessageRpcServer.h b/Engine/Source/Runtime/MessagingRpc/Private/MessageRpcServer.h
new file mode 100644
--- /dev/null
+++ b/Engine/Source/Runtime/MessagingRpc/Private/MessageRpcServer.h
@@ -0,0 +1,76 @@
+// Copyright 1998-2015 Epic Games, Inc. All Rights Reserved.
+
+#pragma once
+
+#include "IMessageRpcServer.h"
+
+
+class IMessageRpcHandler;
+class IMessageRpcReturn;
+
+
+/**
+ * Implements an RPC server.
+ */
+class FMessageRpcServer
+	: public IMessageRpcServer
+{
+public:
+
+	/** Default constructor. */
+	FMessageRpcServer();
+
+	/** Virtual destructor. */
+	virtual ~FMessageRpcServer();
+
+public:
+
+	// IMessageRpcServer interface
+
+	virtual void AddHandler(const FName& RequestMessageType, const TSharedRef<IMessageRpcHandler>& Handler) override;
+
+protected:
+
+	struct FReturnInfo
+	{
+		FMessageAddress ClientAddress;
+		FDateTime LastProgressSent;
+		TSharedPtr<IAsyncProgress> Progress;
+		TSharedPtr<IMessageRpcReturn> Return;
+		TSharedPtr<IAsyncTask> Task;
+	};
+
+	/** Processes an FMessageRpcCancel message. */
+	void ProcessCancelation(const FMessageRpcCancel& Message, const TSharedRef<IMessageContext, ESPMode::ThreadSafe>& Context);
+
+	/** Processes an RPC request message. */
+	void ProcessRequest(const TSharedRef<IMessageContext, ESPMode::ThreadSafe>& Context);
+
+	/** Send a progress message to the RPC client that made the RPC call. */
+	void SendProgress(const FGuid& CallId, const FReturnInfo& ReturnInfo);
+
+	/** Send a result message to the RPC client that made the RPC call. */
+	void SendResult(const FGuid& CallId, const FReturnInfo& ReturnInfo);
+
+private:
+
+	/** Handles all incoming messages. */
+	void HandleMessage(const TSharedRef<IMessageContext, ESPMode::ThreadSafe>& Context);
+
+	/** Handles the ticker. */
+	bool HandleTicker(float DeltaTime);
+
+private:
+
+	/** Registered request message handlers. */
+	TMap<FName, TSharedPtr<IMessageRpcHandler>> Handlers;
+
+	/** Message endpoint. */
+	TSharedPtr<FMessageEndpoint, ESPMode::ThreadSafe> MessageEndpoint;
+
+	/** Collection of pending RPC returns. */
+	TMap<FGuid, FReturnInfo> Returns;
+
+	/** Handle to the registered ticker. */
+	FDelegateHandle TickerHandle;
+};
diff --git a/Engine/Source/Runtime/MessagingRpc/Private/MessagingRpcModule.cpp b/Engine/Source/Runtime/MessagingRpc/Private/MessagingRpcModule.cpp
new file mode 100644
--- /dev/null
+++ b/Engine/Source/Runtime/MessagingRpc/Private/MessagingRpcModule.cpp
@@ -0,0 +1,48 @@
+// Copyright 1998-2015 Epic Games, Inc. All Rights Reserved.
+
+#include "PrivatePCH.h"
+#include "IMessagingRpcModule.h"
+#include "ModuleManager.h"
+
+
+/**
+ * Implements the MessagingRpc module.
+ */
+class FMessagingRpcModule
+	: public IMessagingRpcModule
+{
+public:
+
+	// IMessagingRpcModule interface
+
+	virtual TSharedRef<IMessageRpcClient> CreateRpcClient() override
+	{
+		return MakeShareable(new FMessageRpcClient);
+	}
+
+	virtual TSharedRef<IMessageRpcServer> CreateRpcServer() override
+	{
+		return MakeShareable(new FMessageRpcServer);
+	}
+
+public:
+
+	// IModuleInterface interface
+
+	virtual void StartupModule() override
+	{
+	}
+
+	virtual void ShutdownModule() override
+	{
+	}
+
+	virtual bool SupportsDynamicReloading() override
+	{
+		return false;
+	}
+
+};
+
+
+IMPLEMENT_MODULE(FMessagingRpcModule, MessagingRpc);
diff --git a/Engine/Source/Runtime/MessagingRpc/Private/PrivatePCH.h b/Engine/Source/Runtime/MessagingRpc/Private/PrivatePCH.h
new file mode 100644
--- /dev/null
+++ b/Engine/Source/Runtime/MessagingRpc/Private/PrivatePCH.h
@@ -0,0 +1,34 @@
+// Copyright 1998-2015 Epic Games, Inc. All Rights Reserved.
+
+#pragma once
+
+
+/* Private dependencies
+ *****************************************************************************/
+
+#include "Core.h"
+#include "CoreUObject.h"
+#include "MessageEndpoint.h"
+#include "MessageEndpointBuilder.h"
+#include "IMessageRpcClient.h"
+#include "IMessageRpcServer.h"
+#include "RpcMessage.h"
+#include "Ticker.h"
+
+
+/* Private constants
+ *****************************************************************************/
+
+/** Defines interval at which calls are being re-sent to the server (in seconds). */
+#define MESSAGE_RPC__RETRY_INTERVAL 1.0
+
+/** Defines the time after which calls time out (in seconds). */
+#define MESSAGE_RPC_RETRY_TIMEOUT 3.0
+
+
+/* Private includes
+ *****************************************************************************/
+
+#include "MessageRpcClient.h"
+#include "MessageRpcMessages.h"
+#include "MessageRpcServer.h"
diff --git a/Engine/Source/Runtime/MessagingRpc/Public/IMessageRpcCall.h b/Engine/Source/Runtime/MessagingRpc/Public/IMessageRpcCall.h
new file mode 100644
--- /dev/null
+++ b/Engine/Source/Runtime/MessagingRpc/Public/IMessageRpcCall.h
@@ -0,0 +1,77 @@
+// Copyright 1998-2015 Epic Games, Inc. All Rights Reserved.
+
+#pragma once
+
+
+class IMessageContext;
+class UScriptStruct;
+
+
+/**
+ * Interface for RPC calls.
+ *
+ * Every time an RPC call is made, a request message containing the call parameters is
+ * sent to the remote endpoint. While the remote endpoint is executing the call, it may
+ * send back progress updates in regular intervals. Once the call is complete, the remote
+ * endpoint sends a response message containing the result.
+ */
+class IMessageRpcCall
+{
+public:
+
+	/**
+	 * Complete the request and set its result, if available.
+	 *
+	 * @param ResponseContext The context of the response message.
+	 * @see TimeOut
+	 */
+	virtual void Complete(const TSharedRef<IMessageContext, ESPMode::ThreadSafe>& ResponseContext) = 0;
+
+	/**
+	 * Get the request message.
+	 *
+	 * @return Pointer to the request message.
+	 */
+	virtual void* GetMessage() const = 0;
+
+	/**
+	 * Get the time at which the request was last updated by the server.
+	 *
+	 * @return Update time.
+	 */
+	virtual FDateTime GetLastUpdated() const = 0;
+
+	/**
+	 * Get the type of the request message.
+	 *
+	 * @return Request message type.
+	 */
+	virtual UScriptStruct* GetMessageType() const = 0;
+
+	/**
+	 * Gets the time at which the request was created.
+	 *
+	 * @return Creation time.
+	 */
+	virtual FDateTime GetTimeCreated() const = 0;
+
+	/**
+	 * Time out the request.
+	 *
+	 * @see Complete
+	 */
+	virtual void TimeOut() = 0;
+
+	/**
+	 * Update the current progress.
+	 *
+	 * @param InCompletion The new completion percentage.
+	 * @param InStatusText The new status text.
+	 */
+	virtual void UpdateProgress(float InCompletion, const FText& InStatusText) = 0;
+
+public:
+
+	/** Virtual destructor. */
+	virtual ~IMessageRpcCall() { }
+};
diff --git a/Engine/Source/Runtime/MessagingRpc/Public/IMessageRpcClient.h b/Engine/Source/Runtime/MessagingRpc/Public/IMessageRpcClient.h
new file mode 100644
--- /dev/null
+++ b/Engine/Source/Runtime/MessagingRpc/Public/IMessageRpcClient.h
@@ -0,0 +1,205 @@
+// Copyright 1998-2015 Epic Games, Inc. All Rights Reserved.
+
+#pragma once
+
+#include "AsyncResult.h"
+#include "IAsyncProgress.h"
+#include "IAsyncTask.h"
+#include "IMessageContext.h"
+#include "IMessageRpcCall.h"
+
+
+/**
+ * Interface for RPC clients.
+ */
+class IMessageRpcClient
+{
+	/** Abstract base class for RPC calls. */
+	class FCall
+		: public IAsyncProgress
+		, public IAsyncTask
+		, public IMessageRpcCall
+	{
+	public:
+
+		FCall()
+			: State(EAsyncTaskState::Running)
+			, TimeCreated(FDateTime::UtcNow())
+		{ }
+
+	public:
+
+		FSimpleDelegate& OnCancelled()
+		{
+			return CanceledDelegate;
+		}
+
+	public:
+
+		virtual TOptional<float> GetCompletion() override
+		{
+			return Completion;
+		}
+
+		virtual FText GetStatusText() override
+		{
+			return StatusText;
+		}
+
+		virtual FSimpleDelegate& OnProgressChanged() override
+		{
+			return ProgressChangedDelegate;
+		}
+
+	public:
+
+		virtual void Cancel() override
+		{
+			State = EAsyncTaskState::Cancelled;
+			CanceledDelegate.ExecuteIfBound();
+		}
+
+		virtual EAsyncTaskState GetTaskState() override
+		{
+			return State;
+		}
+
+	public:
+
+		virtual FDateTime GetLastUpdated() const override
+		{
+			return LastUpdated;
+		}
+
+		virtual FDateTime GetTimeCreated() const override
+		{
+			return TimeCreated;
+		}
+
+		virtual void UpdateProgress(float InCompletion, const FText& InStatusText) override
+		{
+			Completion = InCompletion;
+			StatusText = InStatusText;
+			LastUpdated = FDateTime::UtcNow();
+
+			ProgressChangedDelegate.ExecuteIfBound();
+		}
+
+	protected:
+
+		float Completion;
+		FDateTime LastUpdated;
+		EAsyncTaskState State;
+		FText StatusText;
+		FDateTime TimeCreated;
+
+	private:
+
+		FSimpleDelegate CanceledDelegate;
+		FSimpleDelegate ProgressChangedDelegate;
+	};
+
+
+	/** Template for RPC requests. */
+	template<typename RpcType, typename... P>
+	class TCall
+		: public FCall
+	{
+	public:
+
+		TCall(P... Params)
+			: Message(new typename RpcType::FRequest(Params...))
+		{ }
+
+		virtual ~TCall()
+		{
+			delete Message;
+		}
+
+		virtual void Complete(const TSharedRef<IMessageContext, ESPMode::ThreadSafe>& ResponseContext) override
+		{
+			if (ResponseContext->GetMessageTypeInfo() == RpcType::FResponse::StaticStruct())
+			{
+				State = EAsyncTaskState::Completed;
+				Promise.SetValue(static_cast<const typename RpcType::FResponse*>(ResponseContext->GetMessage())->Result);
+			}
+			else
+			{
+				State = EAsyncTaskState::Failed;
+				Promise.SetValue(typename RpcType::FResult());
+			}
+		}
+
+		TFuture<typename RpcType::FResult> GetFuture()
+		{
+			return Promise.GetFuture();
+		}
+
+		virtual void* GetMessage() const override
+		{
+			return Message;
+		}
+
+		virtual UScriptStruct* GetMessageType() const override
+		{
+			return RpcType::FRequest::StaticStruct();
+		}
+
+		virtual void TimeOut() override
+		{
+			State = EAsyncTaskState::Failed;
+			Promise.SetValue(typename RpcType::FResult());
+		}
+
+	private:
+
+		TPromise<typename RpcType::FResult> Promise;
+		typename RpcType::FRequest* Message;
+	};
+
+public:
+
+	/**
+	 * Call a remote procedure.
+	 *
+	 * @param RpcType The RPC type definition.
+	 * @param P The call parameter types.
+	 * @param Params The call parameter list.
+	 */
+	template<typename RpcType, typename... P>
+	TAsyncResult<typename RpcType::FResult> Call(P... Params)
+	{
+		typedef TCall<RpcType, P...> CallType;
+
+		TSharedRef<CallType> Call = MakeShareable(new CallType(Params...));
+		FGuid CallId = AddCall(Call);
+
+		Call->OnCancelled().BindLambda([=]() {
+			CancelCall(CallId);
+		});	
+
+		return TAsyncResult<typename RpcType::FResult>(Call->GetFuture(), Call, Call);
+	}
+
+protected:
+
+	/**
+	 * Add an RPC request.
+	 *
+	 * @param Request The call to add.
+	 * @return The call's unique identifier.
+	 */
+	virtual FGuid AddCall(const TSharedRef<IMessageRpcCall>& Call) = 0;
+
+	/**
+	 * Cancel the specified RPC call.
+	 *
+	 * @param CallId The unique identifier of the call to cancel.
+	 */
+	virtual void CancelCall(const FGuid& CallId) = 0;
+
+public:
+
+	/** Virtual destructor. */
+	virtual ~IMessageRpcClient() { }
+};
diff --git a/Engine/Source/Runtime/MessagingRpc/Public/IMessageRpcHandler.h b/Engine/Source/Runtime/MessagingRpc/Public/IMessageRpcHandler.h
new file mode 100644
--- /dev/null
+++ b/Engine/Source/Runtime/MessagingRpc/Public/IMessageRpcHandler.h
@@ -0,0 +1,28 @@
+// Copyright 1998-2015 Epic Games, Inc. All Rights Reserved.
+
+#pragma once
+
+class IMessageContext;
+class IMessageRpcReturn;
+
+
+/**
+ * Interface for RPC request handlers.
+ */
+class IMessageRpcHandler
+{
+public:
+
+	/**
+	 * Handle an RPC request.
+	 *
+	 * @param Context The context of the RPC request message.
+	 * @return An object for the pending RPC return value.
+	 */
+	virtual TSharedRef<IMessageRpcReturn> HandleRequest(const TSharedRef<IMessageContext, ESPMode::ThreadSafe>& Context) = 0;
+
+public:
+
+	/** Virtual destructor. */
+	virtual ~IMessageRpcHandler() { }
+};
diff --git a/Engine/Source/Runtime/MessagingRpc/Public/IMessageRpcReturn.h b/Engine/Source/Runtime/MessagingRpc/Public/IMessageRpcReturn.h
new file mode 100644
--- /dev/null
+++ b/Engine/Source/Runtime/MessagingRpc/Public/IMessageRpcReturn.h
@@ -0,0 +1,41 @@
+// Copyright 1998-2015 Epic Games, Inc. All Rights Reserved.
+
+#pragma once
+
+struct FRpcMessage;
+class UScriptStruct;
+
+
+/**
+ * Interface for RPC return values.
+ */
+class IMessageRpcReturn
+{
+public:
+
+	/** Cancel the computation of the return value. */
+	virtual void Cancel() = 0;
+
+	/**
+	 */
+	virtual FRpcMessage* CreateResponseMessage() const = 0;
+
+	/**
+	 * Gets the type information for the corresponding RPC response message.
+	 *
+	 * @return The message type information.
+	 */
+	virtual UScriptStruct* GetResponseTypeInfo() const = 0;
+
+	/**
+	 * Check whether the return value is ready.
+	 *
+	 * @return true if the return value is ready, false otherwise.
+	 */
+	virtual bool IsReady() const = 0;
+
+public:
+
+	/** Virtual destructor. */
+	virtual ~IMessageRpcReturn() { }
+};
diff --git a/Engine/Source/Runtime/MessagingRpc/Public/IMessageRpcServer.h b/Engine/Source/Runtime/MessagingRpc/Public/IMessageRpcServer.h
new file mode 100644
--- /dev/null
+++ b/Engine/Source/Runtime/MessagingRpc/Public/IMessageRpcServer.h
@@ -0,0 +1,115 @@
+// Copyright 1998-2015 Epic Games, Inc. All Rights Reserved.
+
+#pragma once
+
+#include "AsyncResult.h"
+#include "IAsyncProgress.h"
+#include "IAsyncTask.h"
+#include "IMessageContext.h"
+#include "IMessageRpcHandler.h"
+#include "IMessageRpcReturn.h"
+
+
+/**
+ * Interface for RPC servers.
+ */
+class IMessageRpcServer
+{
+	/** Template for RPC results. */
+	template<typename RpcType>
+	class TReturn
+		: public IMessageRpcReturn
+	{
+	public:
+
+		TReturn(TAsyncResult<typename RpcType::FResult>&& InResult)
+			: Result(MoveTemp(InResult))
+		{ }
+
+		virtual void Cancel() override
+		{
+			auto Task = Result.GetTask();
+
+			if (Task.IsValid())
+			{
+				Task->Cancel();
+			}
+		}
+
+		virtual FRpcMessage* CreateResponseMessage() const override
+		{
+			const TFuture<typename RpcType::FResult>& Future = Result.GetFuture();
+			check(Future.IsReady());
+
+			return new typename RpcType::FResponse(Future.Get());
+		}
+
+		virtual UScriptStruct* GetResponseTypeInfo() const override
+		{
+			return RpcType::FResponse::StaticStruct();
+		}
+
+		virtual bool IsReady() const override
+		{
+			return Result.GetFuture().IsReady();
+		}
+
+		TAsyncResult<typename RpcType::FResult> Result;
+	};
+
+
+	/** Template for RPC request handlers. */
+	template<typename RpcType, typename HandlerType>
+	class THandler
+		: public IMessageRpcHandler
+	{
+	public:
+
+		typedef TAsyncResult<typename RpcType::FResult>(HandlerType::*FuncType)(const typename RpcType::FRequest&);
+
+		THandler( HandlerType* InHandler, FuncType InFunc )
+			: Handler(InHandler)
+			, Func(InFunc)
+		{
+			check(InHandler != nullptr);
+		}
+
+		virtual TSharedRef<IMessageRpcReturn> HandleRequest(const TSharedRef<IMessageContext, ESPMode::ThreadSafe>& Context) override
+		{
+			return MakeShareable(new TReturn<RpcType>((Handler->*Func)(*static_cast<const typename RpcType::FRequest*>(Context->GetMessage()))));
+		}
+	
+	private:
+
+		HandlerType* Handler;
+		FuncType Func;
+	};
+
+public:
+
+	/**
+	 * Add an RPC request handler.
+	 *
+	 * @param Handler The handler to add.
+	 */
+	virtual void AddHandler(const FName& RequestMessageType, const TSharedRef<IMessageRpcHandler>& Handler) = 0;
+
+public:
+
+	/**
+	 * Register an RPC request handler.
+	 *
+	 * @param Handler The object that will handle the requests.
+	 * @param HandlerFunc The object's request handling function.
+	 */
+	template<typename RpcType, typename HandlerType>
+	void RegisterHandler(HandlerType* Handler, typename THandler<RpcType, HandlerType>::FuncType Func)
+	{
+		AddHandler(RpcType::FRequest::StaticStruct()->GetFName(), MakeShareable(new THandler<RpcType, HandlerType>(Handler, Func)));
+	}
+
+public:
+
+	/** Virtual destructor. */
+	virtual ~IMessageRpcServer() { }
+};
diff --git a/Engine/Source/Runtime/MessagingRpc/Public/IMessagingRpcModule.h b/Engine/Source/Runtime/MessagingRpc/Public/IMessagingRpcModule.h
new file mode 100644
--- /dev/null
+++ b/Engine/Source/Runtime/MessagingRpc/Public/IMessagingRpcModule.h
@@ -0,0 +1,38 @@
+// Copyright 1998-2015 Epic Games, Inc. All Rights Reserved.
+
+#pragma once
+
+#include "ModuleInterface.h"
+
+
+class IMessageRpcClient;
+class IMessageRpcServer;
+
+
+/**
+ * Interface for the MessagingRpc module.
+ */
+class IMessagingRpcModule
+	: public IModuleInterface
+{
+public:
+
+	/**
+	 * Create a client for remote procedure calls.
+	 *
+	 * @return The RPC client.
+	 */
+	virtual TSharedRef<IMessageRpcClient> CreateRpcClient() = 0;
+
+	/**
+	 * Create a server for remote procedure calls.
+	 *
+	 * @return The RPC server.
+	 */
+	virtual TSharedRef<IMessageRpcServer> CreateRpcServer() = 0;
+
+public:
+
+	/** Virtual destructor. */
+	virtual ~IMessagingRpcModule() { }
+};
diff --git a/Engine/Source/Runtime/MessagingRpc/Public/RpcMessage.h b/Engine/Source/Runtime/MessagingRpc/Public/RpcMessage.h
new file mode 100644
--- /dev/null
+++ b/Engine/Source/Runtime/MessagingRpc/Public/RpcMessage.h
@@ -0,0 +1,31 @@
+// Copyright 1998-2015 Epic Games, Inc. All Rights Reserved.
+
+#pragma once
+
+#include "RpcMessage.generated.h"
+
+
+/**
+ * Base type for RPC messages.
+ */
+USTRUCT()
+struct FRpcMessage
+{
+	GENERATED_USTRUCT_BODY()
+
+	/** Correlation identifier for the RPC call that this message refers to. */
+	UPROPERTY()
+	FGuid CallId;
+
+	/** Default constructor. */
+	FRpcMessage() { }
+};
+
+
+#define DECLARE_RPC(RpcType, ResultType) \
+	struct RpcType \
+	{ \
+		typedef RpcType##Request FRequest; \
+		typedef RpcType##Response FResponse; \
+		typedef ResultType FResult; \
+	};
