MessagingRpc: Continued to implement RPC layer

diff --git a/Engine/Source/Runtime/MessagingRpc/MessagingRpc.Build.cs b/Engine/Source/Runtime/MessagingRpc/MessagingRpc.Build.cs
--- a/Engine/Source/Runtime/MessagingRpc/MessagingRpc.Build.cs
+++ b/Engine/Source/Runtime/MessagingRpc/MessagingRpc.Build.cs
@@ -11,7 +11,6 @@ namespace UnrealBuildTool.Rules
 				{
 					"Core",
 					"CoreUObject",
-                    "PortalServices",
 				}
 			);
 
diff --git a/Engine/Source/Runtime/MessagingRpc/Private/MessageRpcClient.cpp b/Engine/Source/Runtime/MessagingRpc/Private/MessageRpcClient.cpp
--- a/Engine/Source/Runtime/MessagingRpc/Private/MessageRpcClient.cpp
+++ b/Engine/Source/Runtime/MessagingRpc/Private/MessageRpcClient.cpp
@@ -11,7 +11,7 @@ FMessageRpcClient::FMessageRpcClient()
 	MessageEndpoint = FMessageEndpoint::Builder("FPortalAppWindowEndpoint")
 		.Handling<FMessageRpcProgress>(this, &FMessageRpcClient::HandleMessage);
 
-	TickerHandle = FTicker::GetCoreTicker().AddTicker(FTickerDelegate::CreateRaw(this, &FMessageRpcClient::HandleTicker), MESSAGE_RPC__RETRY_INTERVAL);
+	TickerHandle = FTicker::GetCoreTicker().AddTicker(FTickerDelegate::CreateRaw(this, &FMessageRpcClient::HandleTicker), MESSAGE_RPC_RETRY_INTERVAL);
 }
 
 
@@ -21,6 +21,28 @@ FMessageRpcClient::~FMessageRpcClient()
 }
 
 
+/* IMessageRpcClient interface
+ *****************************************************************************/
+
+void FMessageRpcClient::Connect(const FMessageAddress& InServerAddress)
+{
+	Disconnect();
+	ServerAddress = InServerAddress;
+}
+
+
+void FMessageRpcClient::Disconnect()
+{
+	for (TMap<FGuid, TSharedPtr<IMessageRpcCall>>::TIterator It(Calls); It; ++It)
+	{
+		It.Value()->TimeOut();
+	}
+
+	Calls.Empty();
+	ServerAddress.Invalidate();
+}
+
+
 /* FMessageRpcClient implementation
  *****************************************************************************/
 
@@ -33,28 +55,27 @@ TSharedPtr<IMessageRpcCall> FMessageRpcClient::FindCall(const TSharedRef<IMessag
 
 void FMessageRpcClient::SendCall(const TSharedPtr<IMessageRpcCall>& Call)
 {
-	MessageEndpoint->Send(
-		Call->GetMessage(),
-		Call->GetMessageType(),
-		nullptr,
-		TArrayBuilder<FMessageAddress>().Add(ServerAddress),
-		FTimespan::Zero(),
-		FDateTime::MaxValue()
-	);
+	if (ServerAddress.IsValid())
+	{
+		MessageEndpoint->Send(
+			Call->GetMessage(),
+			Call->GetMessageType(),
+			nullptr,
+			TArrayBuilder<FMessageAddress>().Add(ServerAddress),
+			FTimespan::Zero(),
+			FDateTime::MaxValue()
+		);
+	}
 }
 
 
 /* IMessageRpcClient interface
  *****************************************************************************/
 
-FGuid FMessageRpcClient::AddCall(const TSharedRef<IMessageRpcCall>& Call)
+void FMessageRpcClient::AddCall(const TSharedRef<IMessageRpcCall>& Call)
 {
-	FGuid CallId = FGuid::NewGuid();
-
-	Calls.Add(CallId, Call);
+	Calls.Add(Call->GetId(), Call);
 	SendCall(Call);
-
-	return CallId;
 }
 
 
@@ -96,7 +117,7 @@ bool FMessageRpcClient::HandleTicker(float DeltaTime)
 			It.RemoveCurrent();
 			Call->TimeOut();
 		}
-		else if (UtcNow - Call->GetLastUpdated() > FTimespan::FromSeconds(MESSAGE_RPC__RETRY_INTERVAL))
+		else if (UtcNow - Call->GetLastUpdated() > FTimespan::FromSeconds(MESSAGE_RPC_RETRY_INTERVAL))
 		{
 			SendCall(Call);
 		}
diff --git a/Engine/Source/Runtime/MessagingRpc/Private/MessageRpcClient.h b/Engine/Source/Runtime/MessagingRpc/Private/MessageRpcClient.h
--- a/Engine/Source/Runtime/MessagingRpc/Private/MessageRpcClient.h
+++ b/Engine/Source/Runtime/MessagingRpc/Private/MessageRpcClient.h
@@ -26,6 +26,13 @@ public:
 	/** Virtual destructor. */
 	virtual ~FMessageRpcClient();
 
+public:
+
+	// IMessageRpcClient interface
+
+	virtual void Connect(const FMessageAddress& InServerAddress) override;
+	virtual void Disconnect() override;
+
 protected:
 
 	/**
@@ -47,7 +54,7 @@ protected:
 
 	// IMessageRpcClient interface
 
-	virtual FGuid AddCall(const TSharedRef<IMessageRpcCall>& Call) override;
+	virtual void AddCall(const TSharedRef<IMessageRpcCall>& Call) override;
 	virtual void CancelCall(const FGuid& CallId) override;
 
 private:
diff --git a/Engine/Source/Runtime/MessagingRpc/Private/MessageRpcServer.cpp b/Engine/Source/Runtime/MessagingRpc/Private/MessageRpcServer.cpp
--- a/Engine/Source/Runtime/MessagingRpc/Private/MessageRpcServer.cpp
+++ b/Engine/Source/Runtime/MessagingRpc/Private/MessageRpcServer.cpp
@@ -32,6 +32,12 @@ void FMessageRpcServer::AddHandler(const FName& RequestMessageType, const TShare
 }
 
 
+const FMessageAddress& FMessageRpcServer::GetAddress() const
+{
+	return MessageEndpoint->GetAddress();
+}
+
+
 /* FMessageRpcServer implementation
  *****************************************************************************/
 
@@ -53,19 +59,20 @@ void FMessageRpcServer::ProcessRequest(const TSharedRef<IMessageContext, ESPMode
 
 	if (Handler.IsValid())
 	{
-		/*
-		TSharedPtr<IMessageRpcServer::IPortalRpcReturn> Return = Handler->HandleRequest(Context);
-
-		if (Return.IsValid())
+		FReturnInfo ReturnInfo;
 		{
-			Returns.Add(Message->CallId, Return);
+			ReturnInfo.ClientAddress = Context->GetSender();
+			ReturnInfo.LastProgressSent = FDateTime::UtcNow();
+			ReturnInfo.Return = Handler->HandleRequest(Context);
+		}
 
-			return;
-		}*/
+		Returns.Add(Message->CallId, ReturnInfo);
+	}
+	else
+	{
+		// notify caller that call was not handled
+		MessageEndpoint->Send(new FMessageRpcInvalid(Message->CallId), Context->GetSender());
 	}
-
-	// notify caller that call was not handled
-	MessageEndpoint->Send(new FMessageRpcInvalid(Message->CallId), Context->GetSender());
 }
 
 
@@ -136,7 +143,7 @@ bool FMessageRpcServer::HandleTicker(float DeltaTime)
 			SendResult(It.Key(), ReturnInfo);
 			It.RemoveCurrent();
 		}
-		else if (UtcNow - ReturnInfo.LastProgressSent > FTimespan::FromSeconds(MESSAGE_RPC__RETRY_INTERVAL * 0.25))
+		else if (UtcNow - ReturnInfo.LastProgressSent > FTimespan::FromSeconds(MESSAGE_RPC_RETRY_INTERVAL * 0.25))
 		{
 			SendProgress(It.Key(), ReturnInfo);
 			ReturnInfo.LastProgressSent = UtcNow;
diff --git a/Engine/Source/Runtime/MessagingRpc/Private/MessageRpcServer.h b/Engine/Source/Runtime/MessagingRpc/Private/MessageRpcServer.h
--- a/Engine/Source/Runtime/MessagingRpc/Private/MessageRpcServer.h
+++ b/Engine/Source/Runtime/MessagingRpc/Private/MessageRpcServer.h
@@ -28,6 +28,7 @@ public:
 	// IMessageRpcServer interface
 
 	virtual void AddHandler(const FName& RequestMessageType, const TSharedRef<IMessageRpcHandler>& Handler) override;
+	virtual const FMessageAddress& GetAddress() const override;
 
 protected:
 
diff --git a/Engine/Source/Runtime/MessagingRpc/Private/PrivatePCH.h b/Engine/Source/Runtime/MessagingRpc/Private/PrivatePCH.h
--- a/Engine/Source/Runtime/MessagingRpc/Private/PrivatePCH.h
+++ b/Engine/Source/Runtime/MessagingRpc/Private/PrivatePCH.h
@@ -20,7 +20,7 @@
  *****************************************************************************/
 
 /** Defines interval at which calls are being re-sent to the server (in seconds). */
-#define MESSAGE_RPC__RETRY_INTERVAL 1.0
+#define MESSAGE_RPC_RETRY_INTERVAL 1.0
 
 /** Defines the time after which calls time out (in seconds). */
 #define MESSAGE_RPC_RETRY_TIMEOUT 3.0
diff --git a/Engine/Source/Runtime/MessagingRpc/Public/IMessageRpcCall.h b/Engine/Source/Runtime/MessagingRpc/Public/IMessageRpcCall.h
--- a/Engine/Source/Runtime/MessagingRpc/Public/IMessageRpcCall.h
+++ b/Engine/Source/Runtime/MessagingRpc/Public/IMessageRpcCall.h
@@ -27,6 +27,13 @@ public:
 	 */
 	virtual void Complete(const TSharedRef<IMessageContext, ESPMode::ThreadSafe>& ResponseContext) = 0;
 
+	/**
+	 * Get the call's unique identifier.
+	 *
+	 * @return Call identifier.
+	 */
+	virtual const FGuid& GetId() const = 0;
+
 	/**
 	 * Get the request message.
 	 *
diff --git a/Engine/Source/Runtime/MessagingRpc/Public/IMessageRpcClient.h b/Engine/Source/Runtime/MessagingRpc/Public/IMessageRpcClient.h
--- a/Engine/Source/Runtime/MessagingRpc/Public/IMessageRpcClient.h
+++ b/Engine/Source/Runtime/MessagingRpc/Public/IMessageRpcClient.h
@@ -107,9 +107,11 @@ class IMessageRpcClient
 	{
 	public:
 
-		TCall(P... Params)
+		TCall(const FGuid& CallId, P... Params)
 			: Message(new typename RpcType::FRequest(Params...))
-		{ }
+		{
+			Message->CallId = CallId;
+		}
 
 		virtual ~TCall()
 		{
@@ -135,6 +137,11 @@ class IMessageRpcClient
 			return Promise.GetFuture();
 		}
 
+		virtual const FGuid& GetId() const override
+		{
+			return Message->CallId;
+		}
+
 		virtual void* GetMessage() const override
 		{
 			return Message;
@@ -157,6 +164,23 @@ class IMessageRpcClient
 		typename RpcType::FRequest* Message;
 	};
 
+public:
+
+	/**
+	 * Connect this client to an RPC server.
+	 *
+	 * @param InServerAddress The RPC server's message address.
+	 * @see Disconnect
+	 */
+	virtual void Connect(const FMessageAddress& InServerAddress) = 0;
+
+	/**
+	 * Disconnect this client from the RPC server.
+	 *
+	 * @see Connect
+	 */
+	virtual void Disconnect() = 0;
+
 public:
 
 	/**
@@ -171,12 +195,15 @@ public:
 	{
 		typedef TCall<RpcType, P...> CallType;
 
-		TSharedRef<CallType> Call = MakeShareable(new CallType(Params...));
-		FGuid CallId = AddCall(Call);
+		FGuid CallId = FGuid::NewGuid();
+		TSharedRef<CallType> Call = MakeShareable(new CallType(CallId, Params...));
+		{
+			AddCall(Call);
 
-		Call->OnCancelled().BindLambda([=]() {
-			CancelCall(CallId);
-		});	
+			Call->OnCancelled().BindLambda([=]() {
+				CancelCall(CallId);
+			});	
+		}
 
 		return TAsyncResult<typename RpcType::FResult>(Call->GetFuture(), Call, Call);
 	}
@@ -187,9 +214,8 @@ protected:
 	 * Add an RPC request.
 	 *
 	 * @param Request The call to add.
-	 * @return The call's unique identifier.
 	 */
-	virtual FGuid AddCall(const TSharedRef<IMessageRpcCall>& Call) = 0;
+	virtual void AddCall(const TSharedRef<IMessageRpcCall>& Call) = 0;
 
 	/**
 	 * Cancel the specified RPC call.
diff --git a/Engine/Source/Runtime/MessagingRpc/Public/IMessageRpcServer.h b/Engine/Source/Runtime/MessagingRpc/Public/IMessageRpcServer.h
--- a/Engine/Source/Runtime/MessagingRpc/Public/IMessageRpcServer.h
+++ b/Engine/Source/Runtime/MessagingRpc/Public/IMessageRpcServer.h
@@ -94,6 +94,13 @@ public:
 	 */
 	virtual void AddHandler(const FName& RequestMessageType, const TSharedRef<IMessageRpcHandler>& Handler) = 0;
 
+	/**
+	 * Gets the server's message address.
+	 *
+	 * @return Message address.
+	 */
+	virtual const FMessageAddress& GetAddress() const = 0;
+
 public:
 
 	/**
