MessagingRpc: Cleaned up module dependencies

diff --git a/Engine/Source/Runtime/MessagingRpc/MessagingRpc.Build.cs b/Engine/Source/Runtime/MessagingRpc/MessagingRpc.Build.cs
--- a/Engine/Source/Runtime/MessagingRpc/MessagingRpc.Build.cs
+++ b/Engine/Source/Runtime/MessagingRpc/MessagingRpc.Build.cs
@@ -6,6 +6,13 @@ namespace UnrealBuildTool.Rules
 	{
 		public MessagingRpc(TargetInfo Target)
 		{
+            DynamicallyLoadedModuleNames.AddRange(
+                new string[]
+                {
+                    "Messaging",
+                }
+            );
+
 			PublicDependencyModuleNames.AddRange(
 				new string[]
 				{
@@ -14,7 +21,7 @@ namespace UnrealBuildTool.Rules
 				}
 			);
 
-            PrivateDependencyModuleNames.AddRange(
+            PrivateIncludePathModuleNames.AddRange(
                 new string[]
                 {
                     "Messaging",
