MessagingRpc: Exposed interface name

diff --git a/Engine/Source/Runtime/MessagingRpc/Public/IMessageRpcClient.h b/Engine/Source/Runtime/MessagingRpc/Public/IMessageRpcClient.h
--- a/Engine/Source/Runtime/MessagingRpc/Public/IMessageRpcClient.h
+++ b/Engine/Source/Runtime/MessagingRpc/Public/IMessageRpcClient.h
@@ -229,3 +229,6 @@ public:
 	/** Virtual destructor. */
 	virtual ~IMessageRpcClient() { }
 };
+
+
+Expose_TNameOf(IMessageRpcClient)
