MessagingRpc: Added callback for missing message handlers

diff --git a/Engine/Source/Runtime/MessagingRpc/Private/MessageRpcMessages.h b/Engine/Source/Runtime/MessagingRpc/Private/MessageRpcMessages.h
--- a/Engine/Source/Runtime/MessagingRpc/Private/MessageRpcMessages.h
+++ b/Engine/Source/Runtime/MessagingRpc/Private/MessageRpcMessages.h
@@ -55,9 +55,9 @@ struct FMessageRpcProgress
 };
 
 
-/** Message for notifying RPC clients that a call was invalid. */
+/** Message for notifying RPC clients that a call was not handled by the server. */
 USTRUCT()
-struct FMessageRpcInvalid
+struct FMessageRpcUnhandled
 {
 	GENERATED_USTRUCT_BODY()
 
@@ -66,10 +66,10 @@ struct FMessageRpcInvalid
 	FGuid CallId;
 
 	/** Default constructor. */
-	FMessageRpcInvalid() { }
+	FMessageRpcUnhandled() { }
 
 	/** Creates and initializes a new instance. */
-	FMessageRpcInvalid(const FGuid& InCallId)
+	FMessageRpcUnhandled(const FGuid& InCallId)
 		: CallId(InCallId)
 	{ }
 };
diff --git a/Engine/Source/Runtime/MessagingRpc/Private/MessageRpcServer.cpp b/Engine/Source/Runtime/MessagingRpc/Private/MessageRpcServer.cpp
--- a/Engine/Source/Runtime/MessagingRpc/Private/MessageRpcServer.cpp
+++ b/Engine/Source/Runtime/MessagingRpc/Private/MessageRpcServer.cpp
@@ -38,6 +38,12 @@ const FMessageAddress& FMessageRpcServer::GetAddress() const
 }
 
 
+FOnMessageRpcNoHandler& FMessageRpcServer::OnNoHandler()
+{
+	return NoHandlerDelegate;
+}
+
+
 /* FMessageRpcServer implementation
  *****************************************************************************/
 
@@ -55,7 +61,19 @@ void FMessageRpcServer::ProcessCancelation(const FMessageRpcCancel& Message, con
 void FMessageRpcServer::ProcessRequest(const TSharedRef<IMessageContext, ESPMode::ThreadSafe>& Context)
 {
 	auto Message = (FRpcMessage*)Context->GetMessage();
-	auto Handler = Handlers.FindRef(Context->GetMessageType());
+	const FName MessageType = Context->GetMessageType();
+
+	if (!Handlers.Contains(MessageType))
+	{
+		if (!NoHandlerDelegate.IsBound())
+		{
+			return;
+		}
+
+		NoHandlerDelegate.Execute(MessageType);
+	}
+
+	auto Handler = Handlers.FindRef(MessageType);
 
 	if (Handler.IsValid())
 	{
@@ -71,7 +89,7 @@ void FMessageRpcServer::ProcessRequest(const TSharedRef<IMessageContext, ESPMode
 	else
 	{
 		// notify caller that call was not handled
-		MessageEndpoint->Send(new FMessageRpcInvalid(Message->CallId), Context->GetSender());
+		MessageEndpoint->Send(new FMessageRpcUnhandled(Message->CallId), Context->GetSender());
 	}
 }
 
diff --git a/Engine/Source/Runtime/MessagingRpc/Private/MessageRpcServer.h b/Engine/Source/Runtime/MessagingRpc/Private/MessageRpcServer.h
--- a/Engine/Source/Runtime/MessagingRpc/Private/MessageRpcServer.h
+++ b/Engine/Source/Runtime/MessagingRpc/Private/MessageRpcServer.h
@@ -29,6 +29,7 @@ public:
 
 	virtual void AddHandler(const FName& RequestMessageType, const TSharedRef<IMessageRpcHandler>& Handler) override;
 	virtual const FMessageAddress& GetAddress() const override;
+	virtual FOnMessageRpcNoHandler& OnNoHandler() override;
 
 protected:
 
@@ -69,6 +70,9 @@ private:
 	/** Message endpoint. */
 	TSharedPtr<FMessageEndpoint, ESPMode::ThreadSafe> MessageEndpoint;
 
+	/* Delegate that is executed when a received RPC message has no registered handler. */
+	FOnMessageRpcNoHandler NoHandlerDelegate;
+
 	/** Collection of pending RPC returns. */
 	TMap<FGuid, FReturnInfo> Returns;
 
diff --git a/Engine/Source/Runtime/MessagingRpc/Public/IMessageRpcServer.h b/Engine/Source/Runtime/MessagingRpc/Public/IMessageRpcServer.h
--- a/Engine/Source/Runtime/MessagingRpc/Public/IMessageRpcServer.h
+++ b/Engine/Source/Runtime/MessagingRpc/Public/IMessageRpcServer.h
@@ -10,6 +10,9 @@
 #include "IMessageRpcReturn.h"
 
 
+/** Delegate type for RPC messages that have no registered handler. */
+DECLARE_DELEGATE_OneParam(FOnMessageRpcNoHandler, const FName& /*MessageType*/)
+
 /**
  * Interface for RPC servers.
  */
@@ -101,6 +104,13 @@ public:
 	 */
 	virtual const FMessageAddress& GetAddress() const = 0;
 
+	/**
+	 * Get a delegate that is executed when a received RPC message has no registered handler.
+	 *
+	 * @return The delegate.
+	 */
+	virtual FOnMessageRpcNoHandler& OnNoHandler() = 0;
+
 public:
 
 	/**
