Replacing PRAGMA_ENABLE_DEPRECATION_WARNINGS with PRAGMA_DISABLE_DEPRECATION_WARNINGS in a mismatched pair of macros

diff --git a/Engine/Source/Runtime/MessagingRpc/Private/MessageRpcServer.cpp b/Engine/Source/Runtime/MessagingRpc/Private/MessageRpcServer.cpp
--- a/Engine/Source/Runtime/MessagingRpc/Private/MessageRpcServer.cpp
+++ b/Engine/Source/Runtime/MessagingRpc/Private/MessageRpcServer.cpp
@@ -106,7 +106,7 @@ PRAGMA_DISABLE_DEPRECATION_WARNINGS
 
 		NoHandlerDelegateWithPathName.ExecuteIfBound(MessageType);
 		NoHandlerDelegate.ExecuteIfBound(MessageType.GetAssetName());
-PRAGMA_DISABLE_DEPRECATION_WARNINGS
+PRAGMA_ENABLE_DEPRECATION_WARNINGS
 	}
 
 	auto Handler = Handlers.FindRef(MessageType);
